Movie Scene Text Support:
- Implemented Movie Scene Text Track to support keying Text properties with localization (UMG, 3DText, etc)

diff --git a/Engine/Plugins/MovieScene/MovieSceneTextTrack/MovieSceneTextTrack.uplugin b/Engine/Plugins/MovieScene/MovieSceneTextTrack/MovieSceneTextTrack.uplugin
new file mode 100644
--- /dev/null
+++ b/Engine/Plugins/MovieScene/MovieSceneTextTrack/MovieSceneTextTrack.uplugin
@@ -0,0 +1,29 @@
+{
+	"FileVersion": 1,
+	"Version": 1,
+	"VersionName": "1.0",
+	"FriendlyName": "MovieSceneTextTrack",
+	"Description": "Adds support to key Text Properties in Movie Scene Sequences",
+	"Category": "Text",
+	"CreatedBy": "Epic Games, Inc.",
+	"CreatedByURL": "https://epicgames.com",
+	"DocsURL": "",
+	"MarketplaceURL": "",
+	"SupportURL": "",
+	"CanContainContent": false,
+	"IsBetaVersion": false,
+	"IsExperimentalVersion": false,
+	"Installed": false,
+	"Modules": [
+		{
+			"Name": "MovieSceneTextTrack",
+			"Type": "Runtime",
+			"LoadingPhase": "Default"
+		},
+		{
+			"Name": "MovieSceneTextTrackEditor",
+			"Type": "Editor",
+			"LoadingPhase": "Default"
+		}
+	]
+}
diff --git a/Engine/Plugins/MovieScene/MovieSceneTextTrack/Source/MovieSceneTextTrack/MovieSceneTextTrack.Build.cs b/Engine/Plugins/MovieScene/MovieSceneTextTrack/Source/MovieSceneTextTrack/MovieSceneTextTrack.Build.cs
new file mode 100644
--- /dev/null
+++ b/Engine/Plugins/MovieScene/MovieSceneTextTrack/Source/MovieSceneTextTrack/MovieSceneTextTrack.Build.cs
@@ -0,0 +1,22 @@
+// Copyright Epic Games, Inc. All Rights Reserved.
+
+using UnrealBuildTool;
+
+public class MovieSceneTextTrack : ModuleRules
+{
+	public MovieSceneTextTrack(ReadOnlyTargetRules Target) : base(Target)
+	{
+		PCHUsage = ModuleRules.PCHUsageMode.UseExplicitOrSharedPCHs;
+
+		PublicDependencyModuleNames.AddRange(
+			new string[]
+			{
+				"Core",
+				"CoreUObject",
+				"Engine",
+				"MovieScene",
+				"MovieSceneTracks",
+			}
+		);
+	}
+}
diff --git a/Engine/Plugins/MovieScene/MovieSceneTextTrack/Source/MovieSceneTextTrack/Private/MovieSceneTextChannel.cpp b/Engine/Plugins/MovieScene/MovieSceneTextTrack/Source/MovieSceneTextTrack/Private/MovieSceneTextChannel.cpp
new file mode 100644
--- /dev/null
+++ b/Engine/Plugins/MovieScene/MovieSceneTextTrack/Source/MovieSceneTextTrack/Private/MovieSceneTextChannel.cpp
@@ -0,0 +1,108 @@
+// Copyright Epic Games, Inc. All Rights Reserved.
+
+#include "MovieSceneTextChannel.h"
+
+#if WITH_EDITOR
+#include "UObject/Package.h"
+#endif
+
+#if WITH_EDITOR
+void FMovieSceneTextChannel::SetPackage(UPackage* InPackage)
+{
+	Package = InPackage;
+}
+
+UPackage* FMovieSceneTextChannel::GetPackage() const
+{
+	return Package.Get();
+}
+#endif
+
+const FText* FMovieSceneTextChannel::Evaluate(FFrameTime InTime) const
+{
+	if (Times.Num())
+	{
+		const int32 Index = FMath::Max(0, Algo::UpperBound(Times, InTime.FrameNumber)-1);
+		return &Values[Index];
+	}
+	return bHasDefaultValue ? &DefaultValue : nullptr;
+}
+
+void FMovieSceneTextChannel::GetKeys(const TRange<FFrameNumber>& WithinRange
+	, TArray<FFrameNumber>* OutKeyTimes
+	, TArray<FKeyHandle>* OutKeyHandles)
+{
+	GetData().GetKeys(WithinRange, OutKeyTimes, OutKeyHandles);
+}
+
+void FMovieSceneTextChannel::GetKeyTimes(TArrayView<const FKeyHandle> InHandles, TArrayView<FFrameNumber> OutKeyTimes)
+{
+	GetData().GetKeyTimes(InHandles, OutKeyTimes);
+}
+
+void FMovieSceneTextChannel::SetKeyTimes(TArrayView<const FKeyHandle> InHandles, TArrayView<const FFrameNumber> InKeyTimes)
+{
+	GetData().SetKeyTimes(InHandles, InKeyTimes);
+}
+
+void FMovieSceneTextChannel::DuplicateKeys(TArrayView<const FKeyHandle> InHandles, TArrayView<FKeyHandle> OutNewHandles)
+{
+	GetData().DuplicateKeys(InHandles, OutNewHandles);
+}
+
+void FMovieSceneTextChannel::DeleteKeys(TArrayView<const FKeyHandle> InHandles)
+{
+	GetData().DeleteKeys(InHandles);
+}
+
+void FMovieSceneTextChannel::DeleteKeysFrom(FFrameNumber InTime, bool bDeleteKeysBefore)
+{
+	// Insert a key at the current time to maintain evaluation
+	if (GetData().GetTimes().Num() > 0)
+	{
+		if (const FText* CurrentValue = Evaluate(InTime))
+		{
+			FText Value(*CurrentValue);
+			GetData().UpdateOrAddKey(InTime, Value);
+		}
+	}
+	GetData().DeleteKeysFrom(InTime, bDeleteKeysBefore);
+}
+
+void FMovieSceneTextChannel::ChangeFrameResolution(FFrameRate SourceRate, FFrameRate DestinationRate)
+{
+	GetData().ChangeFrameResolution(SourceRate, DestinationRate);
+}
+
+TRange<FFrameNumber> FMovieSceneTextChannel::ComputeEffectiveRange() const
+{
+	return GetData().GetTotalRange();
+}
+
+int32 FMovieSceneTextChannel::GetNumKeys() const
+{
+	return Times.Num();
+}
+
+void FMovieSceneTextChannel::Reset()
+{
+	Times.Reset();
+	Values.Reset();
+	KeyHandles.Reset();
+	bHasDefaultValue = false;
+}
+
+void FMovieSceneTextChannel::Offset(FFrameNumber DeltaPosition)
+{
+	GetData().Offset(DeltaPosition);
+}
+
+void FMovieSceneTextChannel::Optimize(const FKeyDataOptimizationParams& InParameters)
+{
+	UE::MovieScene::Optimize(this, InParameters);
+}
+
+void FMovieSceneTextChannel::ClearDefault()
+{
+	bHasDefaultValue = false;
+}
diff --git a/Engine/Plugins/MovieScene/MovieSceneTextTrack/Source/MovieSceneTextTrack/Private/MovieSceneTextPropertySystem.cpp b/Engine/Plugins/MovieScene/MovieSceneTextTrack/Source/MovieSceneTextTrack/Private/MovieSceneTextPropertySystem.cpp
new file mode 100644
--- /dev/null
+++ b/Engine/Plugins/MovieScene/MovieSceneTextTrack/Source/MovieSceneTextTrack/Private/MovieSceneTextPropertySystem.cpp
@@ -0,0 +1,20 @@
+// Copyright Epic Games, Inc. All Rights Reserved.
+
+#include "MovieSceneTextPropertySystem.h"
+#include "TextChannelEvaluatorSystem.h"
+#include "TextComponentTypes.h"
+
+UMovieSceneTextPropertySystem::UMovieSceneTextPropertySystem(const FObjectInitializer& InObjectInitializer)
+	: Super(InObjectInitializer)
+{
+	BindToProperty(UE::MovieScene::FTextComponentTypes::Get()->Text);
+	if (HasAnyFlags(RF_ClassDefaultObject))
+	{
+		DefineImplicitPrerequisite(UTextChannelEvaluatorSystem::StaticClass(), GetClass());
+	}
+}
+
+void UMovieSceneTextPropertySystem::OnRun(FSystemTaskPrerequisites& InPrerequisites, FSystemSubsequentTasks& Subsequents)
+{
+	Super::OnRun(InPrerequisites, Subsequents);
+}
diff --git a/Engine/Plugins/MovieScene/MovieSceneTextTrack/Source/MovieSceneTextTrack/Private/MovieSceneTextSection.cpp b/Engine/Plugins/MovieScene/MovieSceneTextTrack/Source/MovieSceneTextTrack/Private/MovieSceneTextSection.cpp
new file mode 100644
--- /dev/null
+++ b/Engine/Plugins/MovieScene/MovieSceneTextTrack/Source/MovieSceneTextTrack/Private/MovieSceneTextSection.cpp
@@ -0,0 +1,44 @@
+// Copyright Epic Games, Inc. All Rights Reserved.
+
+#include "MovieSceneTextSection.h"
+#include "TextComponentTypes.h"
+#include "Tracks/MovieScenePropertyTrack.h"
+
+UMovieSceneTextSection::UMovieSceneTextSection(const FObjectInitializer& InObjectInitializer)
+	: Super(InObjectInitializer)
+{
+	bSupportsInfiniteRange = true;
+	EvalOptions.EnableAndSetCompletionMode(EMovieSceneCompletionMode::ProjectDefault);
+
+#if WITH_EDITOR
+	ChannelProxy = MakeShared<FMovieSceneChannelProxy>(TextChannel, FMovieSceneChannelMetaData(), TMovieSceneExternalValue<FText>::Make());
+	TextChannel.SetPackage(GetPackage());
+#else
+	ChannelProxy = MakeShared<FMovieSceneChannelProxy>(TextChannel);
+#endif
+}
+
+bool UMovieSceneTextSection::PopulateEvaluationFieldImpl(const TRange<FFrameNumber>& EffectiveRange
+	, const FMovieSceneEvaluationFieldEntityMetaData& InMetaData
+	, FMovieSceneEntityComponentFieldBuilder* OutFieldBuilder)
+{
+	FMovieScenePropertyTrackEntityImportHelper::PopulateEvaluationField(*this, EffectiveRange, InMetaData, OutFieldBuilder);
+	return true;
+}
+
+void UMovieSceneTextSection::ImportEntityImpl(UMovieSceneEntitySystemLinker* EntityLinker
+	, const FEntityImportParams& Params
+	, FImportedEntity* OutImportedEntity)
+{
+	using namespace UE::MovieScene;
+
+	if (!TextChannel.HasAnyData())
+	{
+		return;
+	}
+
+	const FTextComponentTypes* TextComponents = FTextComponentTypes::Get();
+	FPropertyTrackEntityImportHelper(TextComponents->Text)
+		.Add(TextComponents->TextChannel, &TextChannel)
+		.Commit(this, Params, OutImportedEntity);
+}
diff --git a/Engine/Plugins/MovieScene/MovieSceneTextTrack/Source/MovieSceneTextTrack/Private/MovieSceneTextTrack.cpp b/Engine/Plugins/MovieScene/MovieSceneTextTrack/Source/MovieSceneTextTrack/Private/MovieSceneTextTrack.cpp
new file mode 100644
--- /dev/null
+++ b/Engine/Plugins/MovieScene/MovieSceneTextTrack/Source/MovieSceneTextTrack/Private/MovieSceneTextTrack.cpp
@@ -0,0 +1,53 @@
+// Copyright Epic Games, Inc. All Rights Reserved.
+
+#include "MovieSceneTextTrack.h"
+#include "MovieSceneTextSection.h"
+
+UMovieSceneTextTrack::UMovieSceneTextTrack()
+{
+}
+
+bool UMovieSceneTextTrack::SupportsType(TSubclassOf<UMovieSceneSection> SectionClass) const
+{
+	return SectionClass == UMovieSceneTextSection::StaticClass();
+}
+
+UMovieSceneSection* UMovieSceneTextTrack::CreateNewSection()
+{
+	return NewObject<UMovieSceneTextSection>(this, NAME_None, RF_Transactional);
+}
+
+void UMovieSceneTextTrack::AddSection(UMovieSceneSection& Section)
+{
+	Sections.Add(&Section);
+}
+
+const TArray<UMovieSceneSection*>& UMovieSceneTextTrack::GetAllSections() const
+{
+	return Sections;
+}
+
+bool UMovieSceneTextTrack::HasSection(const UMovieSceneSection& Section) const
+{
+	return Sections.Contains(&Section);
+}
+
+bool UMovieSceneTextTrack::IsEmpty() const
+{
+	return Sections.IsEmpty();
+}
+
+void UMovieSceneTextTrack::RemoveAllAnimationData()
+{
+	Sections.Empty();
+}
+
+void UMovieSceneTextTrack::RemoveSection(UMovieSceneSection& Section)
+{
+	Sections.Remove(&Section);
+}
+
+void UMovieSceneTextTrack::RemoveSectionAt(int32 SectionIndex)
+{
+	Sections.RemoveAt(SectionIndex);
+}
diff --git a/Engine/Plugins/MovieScene/MovieSceneTextTrack/Source/MovieSceneTextTrack/Private/MovieSceneTextTrackModule.cpp b/Engine/Plugins/MovieScene/MovieSceneTextTrack/Source/MovieSceneTextTrack/Private/MovieSceneTextTrackModule.cpp
new file mode 100644
--- /dev/null
+++ b/Engine/Plugins/MovieScene/MovieSceneTextTrack/Source/MovieSceneTextTrack/Private/MovieSceneTextTrackModule.cpp
@@ -0,0 +1,24 @@
+// Copyright Epic Games, Inc. All Rights Reserved.
+
+#include "Modules/ModuleInterface.h"
+#include "Modules/ModuleManager.h"
+#include "TextComponentTypes.h"
+
+class FMovieSceneTextTrackModule : public IModuleInterface
+{
+public:
+	virtual void StartupModule() override;
+	virtual void ShutdownModule() override;
+};
+
+IMPLEMENT_MODULE(FMovieSceneTextTrackModule, MovieSceneTextTrack)
+
+void FMovieSceneTextTrackModule::StartupModule()
+{
+	UE::MovieScene::FTextComponentTypes::Get();
+}
+
+void FMovieSceneTextTrackModule::ShutdownModule()
+{
+	UE::MovieScene::FTextComponentTypes::Destroy();
+}
diff --git a/Engine/Plugins/MovieScene/MovieSceneTextTrack/Source/MovieSceneTextTrack/Private/TextChannelEvaluatorSystem.cpp b/Engine/Plugins/MovieScene/MovieSceneTextTrack/Source/MovieSceneTextTrack/Private/TextChannelEvaluatorSystem.cpp
new file mode 100644
--- /dev/null
+++ b/Engine/Plugins/MovieScene/MovieSceneTextTrack/Source/MovieSceneTextTrack/Private/TextChannelEvaluatorSystem.cpp
@@ -0,0 +1,76 @@
+// Copyright Epic Games, Inc. All Rights Reserved.
+
+#include "TextChannelEvaluatorSystem.h"
+#include "EntitySystem/BuiltInComponentTypes.h"
+#include "EntitySystem/MovieSceneEntitySystemLinker.h"
+#include "EntitySystem/MovieSceneEntitySystemTask.h"
+#include "EntitySystem/MovieSceneEvalTimeSystem.h"
+#include "MovieSceneTextChannel.h"
+#include "TextComponentTypes.h"
+
+DECLARE_CYCLE_STAT(TEXT("MovieScene: Evaluate text channels"), MovieSceneEval_EvaluateTextChannelTask, STATGROUP_MovieSceneECS);
+
+namespace UE::MovieScene::Private
+{
+	struct FEvaluateTextChannels
+	{
+		static void ForEachEntity(FSourceTextChannel TextChannel, FFrameTime FrameTime, FText& OutResult)
+		{
+			if (const FText* Value = TextChannel.Source->Evaluate(FrameTime))
+			{
+				OutResult = *Value;
+			}
+			else
+			{
+				OutResult = FText::GetEmpty();
+			}
+		}
+	};
+}
+
+UTextChannelEvaluatorSystem::UTextChannelEvaluatorSystem(const FObjectInitializer& InObjectInitializer)
+	: Super(InObjectInitializer)
+{
+	using namespace UE::MovieScene;
+
+	SystemCategories  = EEntitySystemCategory::ChannelEvaluators;
+	RelevantComponent = FTextComponentTypes::Get()->TextChannel;
+	Phase             = ESystemPhase::Scheduling;
+
+	if (HasAnyFlags(RF_ClassDefaultObject))
+	{
+		DefineImplicitPrerequisite(UMovieSceneEvalTimeSystem::StaticClass(), GetClass());
+	}
+}
+
+void UTextChannelEvaluatorSystem::OnSchedulePersistentTasks(UE::MovieScene::IEntitySystemScheduler* TaskScheduler)
+{
+	using namespace UE::MovieScene;
+
+	const FBuiltInComponentTypes* BuiltInComponents = FBuiltInComponentTypes::Get();
+	const FTextComponentTypes* TextComponents = FTextComponentTypes::Get();
+
+	FEntityTaskBuilder()
+		.Read(TextComponents->TextChannel)
+		.Read(BuiltInComponents->EvalTime)
+		.Write(TextComponents->TextResult)
+		.FilterNone({ BuiltInComponents->Tags.Ignored })
+		.SetStat(GET_STATID(MovieSceneEval_EvaluateTextChannelTask))
+		.Fork_PerEntity<Private::FEvaluateTextChannels>(&Linker->EntityManager, TaskScheduler);
+}
+
+void UTextChannelEvaluatorSystem::OnRun(FSystemTaskPrerequisites& InPrerequisites, FSystemSubsequentTasks& Subsequents)
+{
+	using namespace UE::MovieScene;
+
+	const FBuiltInComponentTypes* BuiltInComponents = FBuiltInComponentTypes::Get();
+	const FTextComponentTypes* TextComponents = FTextComponentTypes::Get();
+
+	FEntityTaskBuilder()
+		.Read(TextComponents->TextChannel)
+		.Read(BuiltInComponents->EvalTime)
+		.Write(TextComponents->TextResult)
+		.FilterNone({ BuiltInComponents->Tags.Ignored })
+		.SetStat(GET_STATID(MovieSceneEval_EvaluateTextChannelTask))
+		.Dispatch_PerEntity<Private::FEvaluateTextChannels>(&Linker->EntityManager, InPrerequisites, &Subsequents);
+}
diff --git a/Engine/Plugins/MovieScene/MovieSceneTextTrack/Source/MovieSceneTextTrack/Private/TextComponentTypes.cpp b/Engine/Plugins/MovieScene/MovieSceneTextTrack/Source/MovieSceneTextTrack/Private/TextComponentTypes.cpp
new file mode 100644
--- /dev/null
+++ b/Engine/Plugins/MovieScene/MovieSceneTextTrack/Source/MovieSceneTextTrack/Private/TextComponentTypes.cpp
@@ -0,0 +1,53 @@
+// Copyright Epic Games, Inc. All Rights Reserved.
+
+#include "TextComponentTypes.h"
+#include "EntitySystem/BuiltInComponentTypes.h"
+#include "EntitySystem/MovieSceneComponentRegistry.h"
+#include "EntitySystem/MovieSceneEntityFactoryTemplates.h"
+#include "EntitySystem/MovieSceneEntitySystemLinker.h"
+#include "EntitySystem/MovieScenePropertyComponentHandler.h"
+
+namespace UE::MovieScene
+{
+
+static TUniquePtr<FTextComponentTypes> GTextComponentTypes;
+static bool GTextComponentTypesDestroyed = false;
+
+FTextComponentTypes* FTextComponentTypes::Get()
+{
+	if (!GTextComponentTypes.IsValid())
+	{
+		check(!GTextComponentTypesDestroyed);
+		GTextComponentTypes.Reset(new FTextComponentTypes);
+	}
+	return GTextComponentTypes.Get();
+}
+
+void FTextComponentTypes::Destroy()
+{
+	GTextComponentTypes.Reset();
+	GTextComponentTypesDestroyed = true;
+}
+
+FTextComponentTypes::FTextComponentTypes()
+{
+	FBuiltInComponentTypes* BuiltInTypes  = FBuiltInComponentTypes::Get();
+	FComponentRegistry* ComponentRegistry = UMovieSceneEntitySystemLinker::GetComponents();
+
+	ComponentRegistry->NewPropertyType(Text, TEXT("FText"));
+
+	ComponentRegistry->NewComponentType(&TextChannel, TEXT("Text Channel"));
+	ComponentRegistry->NewComponentType(&TextResult , TEXT("Text Result"));
+
+	ComponentRegistry->Factories.DuplicateChildComponent(TextChannel);
+	ComponentRegistry->Factories.DuplicateChildComponent(TextResult);
+
+	ComponentRegistry->Factories.DefineMutuallyInclusiveComponent(TextChannel, BuiltInTypes->EvalTime);
+	ComponentRegistry->Factories.DefineMutuallyInclusiveComponent(TextChannel, TextResult);
+
+	BuiltInTypes->PropertyRegistry.DefineProperty(Text, TEXT("Apply Text Properties"))
+		.AddSoleChannel(TextResult)
+		.Commit();
+}
+
+}
diff --git a/Engine/Plugins/MovieScene/MovieSceneTextTrack/Source/MovieSceneTextTrack/Public/MovieSceneTextChannel.h b/Engine/Plugins/MovieScene/MovieSceneTextTrack/Source/MovieSceneTextTrack/Public/MovieSceneTextChannel.h
new file mode 100644
--- /dev/null
+++ b/Engine/Plugins/MovieScene/MovieSceneTextTrack/Source/MovieSceneTextTrack/Public/MovieSceneTextChannel.h
@@ -0,0 +1,170 @@
+// Copyright Epic Games, Inc. All Rights Reserved.
+
+#pragma once
+
+#include "Channels/MovieSceneChannel.h"
+#include "Channels/MovieSceneChannelData.h"
+#include "Channels/MovieSceneChannelTraits.h"
+
+#if WITH_EDITOR
+#include "UObject/WeakObjectPtr.h"
+#endif
+
+#include "MovieSceneTextChannel.generated.h"
+
+#if WITH_EDITOR
+class UPackage;
+#endif
+
+USTRUCT()
+struct MOVIESCENETEXTTRACK_API FMovieSceneTextChannel : public FMovieSceneChannel
+{
+	GENERATED_BODY()
+
+#if WITH_EDITOR
+	void SetPackage(UPackage* InPackage);
+
+	UPackage* GetPackage() const;
+#endif
+
+	/**
+	 * Access a mutable interface for this channel's data
+	 * @return An object that is able to manipulate this channel's data
+	 */
+	FORCEINLINE TMovieSceneChannelData<FText> GetData()
+	{
+		return TMovieSceneChannelData<FText>(&Times, &Values, &KeyHandles);
+	}
+
+	/**
+	 * Access a constant interface for this channel's data
+	 * @return An object that is able to interrogate this channel's data
+	 */
+	FORCEINLINE TMovieSceneChannelData<const FText> GetData() const
+	{
+		return TMovieSceneChannelData<const FText>(&Times, &Values);
+	}
+
+	FORCEINLINE TConstArrayView<FFrameNumber> GetTimes() const
+	{
+		return Times;
+	}
+
+	/** Returns whether this channel has any values */
+	FORCEINLINE bool HasAnyData() const
+	{
+		return Times.Num() != 0 || bHasDefaultValue == true;
+	}
+
+	/**
+	 * Evaluate this channel
+	 * @param InTime     The time to evaluate at
+	 * @return A pointer to the string, or nullptr
+	 */
+	const FText* Evaluate(FFrameTime InTime) const;
+
+	/**
+	 * Add keys with these times to channel. The number of elements in both arrays much match or nothing is added.
+	 * @param InTimes Times to add
+	 * @param InValues Values to add
+	 */
+	FORCEINLINE void AddKeys(const TArray<FFrameNumber>& InTimes, const TArray<FText>& InValues)
+	{
+		check(InTimes.Num() == InValues.Num());
+		Times.Append(InTimes);
+		Values.Append(InValues);
+	}
+
+	/**
+	 * Set this channel's default value that should be used when no keys are present
+	 * @param InDefaultValue The desired default value
+	 */
+	FORCEINLINE void SetDefault(FText InDefaultValue)
+	{
+		bHasDefaultValue = true;
+		DefaultValue = InDefaultValue;
+	}
+
+	/**
+	 * Get this channel's default value that will be used when no keys are present
+	 * @return (Optional) The channel's default value
+	 */
+	FORCEINLINE TOptional<FText> GetDefault() const
+	{
+		return bHasDefaultValue ? TOptional<FText>(DefaultValue) : TOptional<FText>();
+	}
+
+	/** Remove this channel's default value causing the channel to have no effect where no keys are present */
+	FORCEINLINE void RemoveDefault()
+	{
+		bHasDefaultValue = false;
+	}
+
+	//~ Begin FMovieSceneChannel
+	virtual void GetKeys(const TRange<FFrameNumber>& WithinRange, TArray<FFrameNumber>* OutKeyTimes, TArray<FKeyHandle>* OutKeyHandles) override;
+	virtual void GetKeyTimes(TArrayView<const FKeyHandle> InHandles, TArrayView<FFrameNumber> OutKeyTimes) override;
+	virtual void SetKeyTimes(TArrayView<const FKeyHandle> InHandles, TArrayView<const FFrameNumber> InKeyTimes) override;
+	virtual void DuplicateKeys(TArrayView<const FKeyHandle> InHandles, TArrayView<FKeyHandle> OutNewHandles) override;
+	virtual void DeleteKeys(TArrayView<const FKeyHandle> InHandles) override;
+	virtual void DeleteKeysFrom(FFrameNumber InTime, bool bDeleteKeysBefore) override;
+	virtual void ChangeFrameResolution(FFrameRate SourceRate, FFrameRate DestinationRate) override;
+	virtual TRange<FFrameNumber> ComputeEffectiveRange() const override;
+	virtual int32 GetNumKeys() const override;
+	virtual void Reset() override;
+	virtual void Offset(FFrameNumber DeltaPosition) override;
+	virtual void Optimize(const FKeyDataOptimizationParams& InParameters) override;
+	virtual void ClearDefault() override;
+	//~ End FMovieSceneChannel
+
+private:
+	UPROPERTY(meta=(KeyTimes))
+	TArray<FFrameNumber> Times;
+
+	/** Array of values that correspond to each key time */
+	UPROPERTY(meta=(KeyValues))
+	TArray<FText> Values;
+
+	/** Default value used when there are no keys */
+	UPROPERTY()
+	FText DefaultValue;
+
+	UPROPERTY()
+	bool bHasDefaultValue = false;
+
+	/** This needs to be a UPROPERTY so it gets saved into editor transactions but transient so it does not get saved into assets. */
+	UPROPERTY(Transient)
+	FMovieSceneKeyHandleMap KeyHandles;
+
+#if WITH_EDITOR
+	/** Package to use for Text Localization. Always set to the Package of the Section holding this Channel. */
+	TWeakObjectPtr<UPackage> Package;
+#endif
+};
+
+template<>
+struct TMovieSceneChannelTraits<FMovieSceneTextChannel> : TMovieSceneChannelTraitsBase<FMovieSceneTextChannel>
+{
+#if WITH_EDITOR
+	/** Text channels can have external values (ie, they can get their values from external objects for UI purposes) */
+	typedef TMovieSceneExternalValue<FText> ExtendedEditorDataType;
+#endif
+};
+
+inline bool EvaluateChannel(const FMovieSceneTextChannel* InChannel, FFrameTime InTime, FText& OutValue)
+{
+	if (const FText* Result = InChannel->Evaluate(InTime))
+	{
+		OutValue = *Result;
+		return true;
+	}
+	return false;
+}
+
+inline bool ValueExistsAtTime(const FMovieSceneTextChannel* Channel, FFrameNumber InFrameNumber, const FText& Value)
+{
+	const FFrameTime FrameTime(InFrameNumber);
+
+	const FText* ExistingValue = Channel->Evaluate(FrameTime);
+
+	return ExistingValue && Value.EqualTo(*ExistingValue);
+}
diff --git a/Engine/Plugins/MovieScene/MovieSceneTextTrack/Source/MovieSceneTextTrack/Public/MovieSceneTextPropertySystem.h b/Engine/Plugins/MovieScene/MovieSceneTextTrack/Source/MovieSceneTextTrack/Public/MovieSceneTextPropertySystem.h
new file mode 100644
--- /dev/null
+++ b/Engine/Plugins/MovieScene/MovieSceneTextTrack/Source/MovieSceneTextTrack/Public/MovieSceneTextPropertySystem.h
@@ -0,0 +1,19 @@
+// Copyright Epic Games, Inc. All Rights Reserved.
+
+#pragma once
+
+#include "Systems/MovieScenePropertySystem.h"
+#include "MovieSceneTextPropertySystem.generated.h"
+
+UCLASS(MinimalAPI)
+class UMovieSceneTextPropertySystem : public UMovieScenePropertySystem
+{
+	GENERATED_BODY()
+
+public:
+	UMovieSceneTextPropertySystem(const FObjectInitializer& InObjectInitializer);
+
+	//~ Begin UMovieSceneEntitySystem
+	virtual void OnRun(FSystemTaskPrerequisites& InPrerequisites, FSystemSubsequentTasks& Subsequents) override;
+	//~ End UMovieSceneEntitySystem
+};
diff --git a/Engine/Plugins/MovieScene/MovieSceneTextTrack/Source/MovieSceneTextTrack/Public/MovieSceneTextSection.h b/Engine/Plugins/MovieScene/MovieSceneTextTrack/Source/MovieSceneTextTrack/Public/MovieSceneTextSection.h
new file mode 100644
--- /dev/null
+++ b/Engine/Plugins/MovieScene/MovieSceneTextTrack/Source/MovieSceneTextTrack/Public/MovieSceneTextSection.h
@@ -0,0 +1,31 @@
+// Copyright Epic Games, Inc. All Rights Reserved.
+
+#pragma once
+
+#include "EntitySystem/IMovieSceneEntityProvider.h"
+#include "MovieSceneSection.h"
+#include "MovieSceneTextChannel.h"
+#include "MovieSceneTextSection.generated.h"
+
+UCLASS(MinimalAPI)
+class UMovieSceneTextSection
+	: public UMovieSceneSection
+	, public IMovieSceneEntityProvider
+{
+	GENERATED_BODY()
+
+public:
+	UMovieSceneTextSection(const FObjectInitializer& InObjectInitializer);
+
+	/** Public access to this section's internal data function */
+	const FMovieSceneTextChannel& GetChannel() const { return TextChannel; }
+
+	//~ Begin IMovieSceneEntityProvider
+	virtual bool PopulateEvaluationFieldImpl(const TRange<FFrameNumber>& EffectiveRange, const FMovieSceneEvaluationFieldEntityMetaData& InMetaData, FMovieSceneEntityComponentFieldBuilder* OutFieldBuilder) override;
+	virtual void ImportEntityImpl(UMovieSceneEntitySystemLinker* EntityLinker, const FEntityImportParams& Params, FImportedEntity* OutImportedEntity) override;
+	//~ End IMovieSceneEntityProvider
+
+private:
+	UPROPERTY()
+	FMovieSceneTextChannel TextChannel;
+};
diff --git a/Engine/Plugins/MovieScene/MovieSceneTextTrack/Source/MovieSceneTextTrack/Public/MovieSceneTextTrack.h b/Engine/Plugins/MovieScene/MovieSceneTextTrack/Source/MovieSceneTextTrack/Public/MovieSceneTextTrack.h
new file mode 100644
--- /dev/null
+++ b/Engine/Plugins/MovieScene/MovieSceneTextTrack/Source/MovieSceneTextTrack/Public/MovieSceneTextTrack.h
@@ -0,0 +1,28 @@
+// Copyright Epic Games, Inc. All Rights Reserved.
+
+#pragma once
+
+#include "Tracks/MovieScenePropertyTrack.h"
+#include "MovieSceneTextTrack.generated.h"
+
+/** Movie Scene Track that holds a series of texts. */
+UCLASS(MinimalAPI)
+class UMovieSceneTextTrack : public UMovieScenePropertyTrack
+{
+	GENERATED_BODY()
+
+public:
+	UMovieSceneTextTrack();
+
+	//~ Begin UMovieSceneTrack
+	virtual bool SupportsType(TSubclassOf<UMovieSceneSection> SectionClass) const override;
+	virtual UMovieSceneSection* CreateNewSection() override;
+	virtual void AddSection(UMovieSceneSection& Section) override;
+	virtual const TArray<UMovieSceneSection*>& GetAllSections() const override;
+	virtual bool HasSection(const UMovieSceneSection& Section) const override;
+	virtual bool IsEmpty() const override;
+	virtual void RemoveAllAnimationData() override;
+	virtual void RemoveSection(UMovieSceneSection& Section) override;
+	virtual void RemoveSectionAt(int32 SectionIndex) override;
+	//~ End UMovieSceneTrack
+};
diff --git a/Engine/Plugins/MovieScene/MovieSceneTextTrack/Source/MovieSceneTextTrack/Public/TextChannelEvaluatorSystem.h b/Engine/Plugins/MovieScene/MovieSceneTextTrack/Source/MovieSceneTextTrack/Public/TextChannelEvaluatorSystem.h
new file mode 100644
--- /dev/null
+++ b/Engine/Plugins/MovieScene/MovieSceneTextTrack/Source/MovieSceneTextTrack/Public/TextChannelEvaluatorSystem.h
@@ -0,0 +1,21 @@
+// Copyright Epic Games, Inc. All Rights Reserved.
+
+#pragma once
+
+#include "EntitySystem/MovieSceneEntitySystem.h"
+#include "TextChannelEvaluatorSystem.generated.h"
+
+/** System responsible for evaluating Text channels */
+UCLASS(MinimalAPI)
+class UTextChannelEvaluatorSystem : public UMovieSceneEntitySystem
+{
+	GENERATED_BODY()
+
+public:
+	UTextChannelEvaluatorSystem(const FObjectInitializer& InObjectInitializer);
+
+	//~ Begin UMovieSceneEntitySystem
+	virtual void OnSchedulePersistentTasks(UE::MovieScene::IEntitySystemScheduler* TaskScheduler) override;
+	virtual void OnRun(FSystemTaskPrerequisites& InPrerequisites, FSystemSubsequentTasks& Subsequents) override;
+	//~ End UMovieSceneEntitySystem
+};
diff --git a/Engine/Plugins/MovieScene/MovieSceneTextTrack/Source/MovieSceneTextTrack/Public/TextComponentTypes.h b/Engine/Plugins/MovieScene/MovieSceneTextTrack/Source/MovieSceneTextTrack/Public/TextComponentTypes.h
new file mode 100644
--- /dev/null
+++ b/Engine/Plugins/MovieScene/MovieSceneTextTrack/Source/MovieSceneTextTrack/Public/TextComponentTypes.h
@@ -0,0 +1,43 @@
+// Copyright Epic Games, Inc. All Rights Reserved.
+
+#pragma once
+
+#include "EntitySystem/MovieSceneEntityIDs.h"
+#include "EntitySystem/MovieScenePropertyMetaDataTraits.h"
+#include "EntitySystem/MovieScenePropertySystemTypes.h"
+#include "EntitySystem/MovieScenePropertyTraits.h"
+#include "Internationalization/Text.h"
+
+struct FMovieSceneTextChannel;
+
+namespace UE::MovieScene
+{
+
+/** The component data for evaluating a Text channel */
+struct FSourceTextChannel
+{
+	FSourceTextChannel() : Source(nullptr) {}
+	FSourceTextChannel(const FMovieSceneTextChannel* InSource) : Source(InSource) {}
+	const FMovieSceneTextChannel* Source;
+};
+
+using FTextPropertyTraits = TDirectPropertyTraits<FText>;
+
+struct FTextComponentTypes
+{
+	MOVIESCENETEXTTRACK_API static FTextComponentTypes* Get();
+	MOVIESCENETEXTTRACK_API static void Destroy();
+
+	// An FMovieSceneTextChannel
+	TComponentTypeID<FSourceTextChannel> TextChannel;
+
+	// Result of an evaluated FMovieSceneTextChannel
+	TComponentTypeID<FText> TextResult;
+
+	TPropertyComponents<FTextPropertyTraits> Text;
+
+private:
+	FTextComponentTypes();
+};
+
+}
\ No newline at end of file
diff --git a/Engine/Plugins/MovieScene/MovieSceneTextTrack/Source/MovieSceneTextTrackEditor/MovieSceneTextTrackEditor.Build.cs b/Engine/Plugins/MovieScene/MovieSceneTextTrack/Source/MovieSceneTextTrackEditor/MovieSceneTextTrackEditor.Build.cs
new file mode 100644
--- /dev/null
+++ b/Engine/Plugins/MovieScene/MovieSceneTextTrack/Source/MovieSceneTextTrackEditor/MovieSceneTextTrackEditor.Build.cs
@@ -0,0 +1,29 @@
+// Copyright Epic Games, Inc. All Rights Reserved.
+
+using UnrealBuildTool;
+
+public class MovieSceneTextTrackEditor : ModuleRules
+{
+	public MovieSceneTextTrackEditor(ReadOnlyTargetRules Target) : base(Target)
+	{
+		PCHUsage = ModuleRules.PCHUsageMode.UseExplicitOrSharedPCHs;
+
+		PrivateDependencyModuleNames.AddRange(
+			new string[]
+			{
+				"Core",
+				"CoreUObject",
+				"Engine",
+				"MovieScene",
+				"MovieSceneTextTrack",
+				"MovieSceneTools",
+				"MovieSceneTracks",
+				"Sequencer",
+				"SequencerCore",
+				"Slate",
+				"SlateCore",
+				"UnrealEd",
+			}
+		);
+	}
+}
diff --git a/Engine/Plugins/MovieScene/MovieSceneTextTrack/Source/MovieSceneTextTrackEditor/Private/MovieSceneTextKeyStruct.h b/Engine/Plugins/MovieScene/MovieSceneTextTrack/Source/MovieSceneTextTrackEditor/Private/MovieSceneTextKeyStruct.h
new file mode 100644
--- /dev/null
+++ b/Engine/Plugins/MovieScene/MovieSceneTextTrack/Source/MovieSceneTextTrackEditor/Private/MovieSceneTextKeyStruct.h
@@ -0,0 +1,22 @@
+// Copyright Epic Games, Inc. All Rights Reserved.
+
+#pragma once
+
+#include "SequencerKeyStructGenerator.h"
+#include "MovieSceneTextKeyStruct.generated.h"
+
+/**
+ * Derived only to explicitly specify that it's Editor Only
+ * to be ignored in Cooked Builds as these KeyStructs are instanced under a non-editor package for text localization
+ * @see InstanceGeneratedStruct overload for FMovieSceneTextChannel
+ */
+UCLASS()
+class UMovieSceneTextKeyStruct : public UMovieSceneKeyStructType
+{
+	GENERATED_BODY()
+
+public:
+	//~ Begin UObject
+	virtual bool IsEditorOnly() const override { return true; }
+	//~ End UObject
+};
diff --git a/Engine/Plugins/MovieScene/MovieSceneTextTrack/Source/MovieSceneTextTrackEditor/Private/MovieSceneTextTrackEditorModule.cpp b/Engine/Plugins/MovieScene/MovieSceneTextTrack/Source/MovieSceneTextTrackEditor/Private/MovieSceneTextTrackEditorModule.cpp
new file mode 100644
--- /dev/null
+++ b/Engine/Plugins/MovieScene/MovieSceneTextTrack/Source/MovieSceneTextTrackEditor/Private/MovieSceneTextTrackEditorModule.cpp
@@ -0,0 +1,47 @@
+// Copyright Epic Games, Inc. All Rights Reserved.
+
+#include "ISequencerModule.h"
+#include "Modules/ModuleInterface.h"
+#include "Modules/ModuleManager.h"
+#include "MovieSceneTextChannel.h"
+#include "SequencerChannelInterface.h"
+#include "TextClipboardTypes.h"
+#include "TextPropertyTrackEditor.h"
+
+class FMovieSceneTextTrackEditorModule : public IModuleInterface
+{
+public:
+	//~ Begin IModuleInterface
+	virtual void StartupModule() override;
+	virtual void ShutdownModule() override;
+	//~ End IModuleInterface
+
+private:
+	FDelegateHandle TextPropertyTrackCreateEditorHandle;
+};
+
+IMPLEMENT_MODULE(FMovieSceneTextTrackEditorModule, MovieSceneTextTrackEditor)
+
+void FMovieSceneTextTrackEditorModule::StartupModule()
+{
+	if (!GIsEditor)
+	{
+		return;
+	}
+
+	ISequencerModule& SequencerModule = FModuleManager::Get().LoadModuleChecked<ISequencerModule>("Sequencer");
+	TextPropertyTrackCreateEditorHandle = SequencerModule.RegisterPropertyTrackEditor<FTextPropertyTrackEditor>();
+	SequencerModule.RegisterChannelInterface<FMovieSceneTextChannel>();
+}
+
+void FMovieSceneTextTrackEditorModule::ShutdownModule()
+{
+	ISequencerModule* SequencerModule = FModuleManager::Get().GetModulePtr<ISequencerModule>("Sequencer");
+	if (!SequencerModule)
+	{
+		return;
+	}
+
+	SequencerModule->UnRegisterTrackEditor(TextPropertyTrackCreateEditorHandle);
+	TextPropertyTrackCreateEditorHandle.Reset();
+}
diff --git a/Engine/Plugins/MovieScene/MovieSceneTextTrack/Source/MovieSceneTextTrackEditor/Private/TextClipboardTypes.h b/Engine/Plugins/MovieScene/MovieSceneTextTrack/Source/MovieSceneTextTrackEditor/Private/TextClipboardTypes.h
new file mode 100644
--- /dev/null
+++ b/Engine/Plugins/MovieScene/MovieSceneTextTrack/Source/MovieSceneTextTrackEditor/Private/TextClipboardTypes.h
@@ -0,0 +1,16 @@
+// Copyright Epic Games, Inc. All Rights Reserved.
+
+#pragma once
+
+#include "Internationalization/Text.h"
+#include "MovieSceneClipboard.h"
+#include "UObject/NameTypes.h"
+
+namespace MovieSceneClipboard
+{
+	template<> inline FName GetKeyTypeName<FText>()
+	{
+		static FName Name("Text");
+		return Name;
+	}
+}
diff --git a/Engine/Plugins/MovieScene/MovieSceneTextTrack/Source/MovieSceneTextTrackEditor/Private/TextPropertyTrackEditor.cpp b/Engine/Plugins/MovieScene/MovieSceneTextTrack/Source/MovieSceneTextTrackEditor/Private/TextPropertyTrackEditor.cpp
new file mode 100644
--- /dev/null
+++ b/Engine/Plugins/MovieScene/MovieSceneTextTrack/Source/MovieSceneTextTrackEditor/Private/TextPropertyTrackEditor.cpp
@@ -0,0 +1,137 @@
+// Copyright Epic Games, Inc. All Rights Reserved.
+
+#include "TextPropertyTrackEditor.h"
+#include "MovieSceneTextChannel.h"
+#include "MovieSceneTextKeyStruct.h"
+#include "SequencerKeyStructGenerator.h"
+#include "UObject/TextProperty.h"
+#include "Widgets/STextKeyEditor.h"
+
+TSharedRef<ISequencerTrackEditor> FTextPropertyTrackEditor::CreateTrackEditor(TSharedRef<ISequencer> OwningSequencer)
+{
+	return MakeShared<FTextPropertyTrackEditor>(OwningSequencer);
+}
+
+void FTextPropertyTrackEditor::GenerateKeysFromPropertyChanged(const FPropertyChangedParams& PropertyChangedParams
+	, UMovieSceneSection* SectionToKey
+	, FGeneratedTrackKeys& OutGeneratedKeys)
+{
+	const FTextProperty* TextProperty = CastField<const FTextProperty>(PropertyChangedParams.PropertyPath.GetLeafMostProperty().Property.Get());
+	if (!TextProperty)
+	{
+		return;
+	}
+
+	FText TextPropertyValue = PropertyChangedParams.GetPropertyValue<FText>();
+	OutGeneratedKeys.Add(FMovieSceneChannelValueSetter::Create<FMovieSceneTextChannel>(0, MoveTemp(TextPropertyValue), true));
+}
+
+bool CanCreateKeyEditor(const FMovieSceneTextChannel* Channel)
+{
+	return true;
+}
+
+TSharedRef<SWidget> CreateKeyEditor(const TMovieSceneChannelHandle<FMovieSceneTextChannel>& Channel
+	, UMovieSceneSection* Section
+	, const FGuid& InObjectBindingID
+	, TWeakPtr<FTrackInstancePropertyBindings> PropertyBindings
+	, TWeakPtr<ISequencer> InSequencer)
+{
+	using namespace UE::MovieScene;
+
+	const TMovieSceneExternalValue<FText>* ExternalValue = Channel.GetExtendedEditorData();
+	if (!ExternalValue)
+	{
+		return SNullWidget::NullWidget;
+	}
+
+	FTextKeyEditorParams Params;
+	Params.ObjectBindingID = InObjectBindingID;
+	Params.ChannelHandle = Channel;
+	Params.WeakSection = Section;
+	Params.WeakSequencer = InSequencer;
+	Params.WeakPropertyBindings = PropertyBindings;
+	Params.OnGetExternalValue = ExternalValue->OnGetExternalValue;
+
+	return SNew(STextKeyEditor, MoveTemp(Params));
+}
+
+UMovieSceneKeyStructType* InstanceGeneratedStruct(FMovieSceneTextChannel* Channel, FSequencerKeyStructGenerator* Generator)
+{
+	UScriptStruct* ChannelType = FMovieSceneTextChannel::StaticStruct();
+	check(ChannelType);
+
+	UPackage* const Package = Channel->GetPackage();
+	if (!ensureMsgf(Package, TEXT("No valid package could be found for Text Channel")))
+	{
+		return nullptr;
+	}
+
+	// Generated struct to be unique per Package
+	const FName InstanceStructName = *(ChannelType->GetName() + Package->GetPersistentGuid().ToString());
+
+	UMovieSceneKeyStructType* Existing = Generator->FindGeneratedStruct(InstanceStructName);
+	if (Existing)
+	{
+		return Existing;
+	}
+
+	static const FName TimesMetaDataTag("KeyTimes");
+	static const FName ValuesMetaDataTag("KeyValues");
+
+	FArrayProperty* SourceTimes  = FSequencerKeyStructGenerator::FindArrayPropertyWithTag(ChannelType, TimesMetaDataTag);
+	FArrayProperty* SourceValues = FSequencerKeyStructGenerator::FindArrayPropertyWithTag(ChannelType, ValuesMetaDataTag);
+
+	if (!ensureAlwaysMsgf(SourceTimes && SourceValues
+		, TEXT("FMovieSceneTextChannel does not have KeyTimes & KeyValues meta data tags in their respective array properties. Did something change ?")))
+	{
+		return nullptr;
+	}
+
+	constexpr EObjectFlags ObjectFlags = RF_Public | RF_Standalone | RF_Transient;
+
+	UMovieSceneTextKeyStruct* NewStruct = NewObject<UMovieSceneTextKeyStruct>(Package, NAME_None, ObjectFlags);
+	NewStruct->SourceTimesProperty  = SourceTimes;
+	NewStruct->SourceValuesProperty = SourceValues;
+
+	FTextProperty* NewValueProperty = new FTextProperty(NewStruct, "Value", RF_NoFlags);
+	NewValueProperty->SetPropertyFlags(CPF_Edit);
+	NewValueProperty->SetMetaData("Category", TEXT("Key"));
+	NewValueProperty->ArrayDim = 1;
+
+	NewStruct->AddCppProperty(NewValueProperty);
+	NewStruct->DestValueProperty = NewValueProperty;
+
+	FSequencerKeyStructGenerator::FinalizeNewKeyStruct(NewStruct);
+
+	Generator->AddGeneratedStruct(InstanceStructName, NewStruct);
+	return NewStruct;
+}
+
+void PostConstructKeyInstance(const TMovieSceneChannelHandle<FMovieSceneTextChannel>& ChannelHandle
+	, FKeyHandle InHandle
+	, FStructOnScope* Struct)
+{
+	FMovieSceneTextChannel* Channel = ChannelHandle.Get();
+	if (!Channel)
+	{
+		return;
+	}
+
+	// Set Package that the Text Property will need for Localization
+	Struct->SetPackage(ChannelHandle.Get()->GetPackage());
+
+	const UMovieSceneKeyStructType* GeneratedStructType = CastChecked<const UMovieSceneKeyStructType>(Struct->GetStruct());
+
+	const int32 InitialKeyIndex = Channel->GetData().GetIndex(InHandle);
+
+	// Copy the initial value into the struct
+	if (InitialKeyIndex != INDEX_NONE)
+	{
+		const uint8* SrcValueData  = GeneratedStructType->SourceValuesProperty->ContainerPtrToValuePtr<uint8>(Channel);
+		uint8* DestValueData = GeneratedStructType->DestValueProperty->ContainerPtrToValuePtr<uint8>(Struct->GetStructMemory());
+
+		FScriptArrayHelper SourceValuesArray(GeneratedStructType->SourceValuesProperty.Get(), SrcValueData);
+		GeneratedStructType->SourceValuesProperty->Inner->CopyCompleteValue(DestValueData, SourceValuesArray.GetRawPtr(InitialKeyIndex));
+	}
+}
diff --git a/Engine/Plugins/MovieScene/MovieSceneTextTrack/Source/MovieSceneTextTrackEditor/Private/TextPropertyTrackEditor.h b/Engine/Plugins/MovieScene/MovieSceneTextTrack/Source/MovieSceneTextTrackEditor/Private/TextPropertyTrackEditor.h
new file mode 100644
--- /dev/null
+++ b/Engine/Plugins/MovieScene/MovieSceneTextTrack/Source/MovieSceneTextTrackEditor/Private/TextPropertyTrackEditor.h
@@ -0,0 +1,52 @@
+// Copyright Epic Games, Inc. All Rights Reserved.
+
+#pragma once
+
+#include "MovieSceneTextTrack.h"
+#include "PropertyTrackEditor.h"
+
+class UMovieSceneKeyStructType;
+class FSequencerKeyStructGenerator;
+struct FMovieSceneTextChannel;
+
+/** Track Editor for Text Property */
+class FTextPropertyTrackEditor : public FPropertyTrackEditor<UMovieSceneTextTrack>
+{
+public:
+	FTextPropertyTrackEditor(TSharedRef<ISequencer> InSequencer)
+		: FPropertyTrackEditor(InSequencer, GetAnimatedPropertyTypes())
+	{
+	}
+
+	/** Retrieve a list of all property types that this track editor animates  */
+	static TArray<FAnimatedPropertyKey, TInlineAllocator<1>> GetAnimatedPropertyTypes()
+	{
+		return TArray<FAnimatedPropertyKey, TInlineAllocator<1>>({ FAnimatedPropertyKey::FromPropertyTypeName(NAME_TextProperty) });
+	}
+
+	/**
+	 * Creates an instance of this class (called by a sequencer).
+	 * @param OwningSequencer The sequencer instance to be used by this tool
+	 * @return The new instance of this class
+	 */
+	static TSharedRef<ISequencerTrackEditor> CreateTrackEditor(TSharedRef<ISequencer> OwningSequencer);
+
+protected:
+	//~ Begin FPropertyTrackEditor
+	virtual void GenerateKeysFromPropertyChanged(const FPropertyChangedParams& PropertyChangedParams, UMovieSceneSection* SectionToKey, FGeneratedTrackKeys& OutGeneratedKeys) override;
+	//~ End FPropertyTrackEditor
+};
+
+bool CanCreateKeyEditor(const FMovieSceneTextChannel* Channel);
+
+TSharedRef<SWidget> CreateKeyEditor(const TMovieSceneChannelHandle<FMovieSceneTextChannel>& Channel
+	, UMovieSceneSection* Section
+	, const FGuid& InObjectBindingID
+	, TWeakPtr<FTrackInstancePropertyBindings> PropertyBindings
+	, TWeakPtr<ISequencer> InSequencer);
+
+UMovieSceneKeyStructType* InstanceGeneratedStruct(FMovieSceneTextChannel* Channel, FSequencerKeyStructGenerator* Generator);
+
+void PostConstructKeyInstance(const TMovieSceneChannelHandle<FMovieSceneTextChannel>& ChannelHandle
+	, FKeyHandle InHandle
+	, FStructOnScope* Struct);
diff --git a/Engine/Plugins/MovieScene/MovieSceneTextTrack/Source/MovieSceneTextTrackEditor/Private/Widgets/STextKeyEditor.cpp b/Engine/Plugins/MovieScene/MovieSceneTextTrack/Source/MovieSceneTextTrackEditor/Private/Widgets/STextKeyEditor.cpp
new file mode 100644
--- /dev/null
+++ b/Engine/Plugins/MovieScene/MovieSceneTextTrack/Source/MovieSceneTextTrackEditor/Private/Widgets/STextKeyEditor.cpp
@@ -0,0 +1,171 @@
+// Copyright Epic Games, Inc. All Rights Reserved.
+
+#include "STextKeyEditor.h"
+#include "IKeyArea.h"
+#include "ISequencer.h"
+#include "MVVM/Selection/Selection.h"
+#include "MVVM/ViewModels/SequencerEditorViewModel.h"
+#include "MovieSceneTimeHelpers.h"
+#include "ScopedTransaction.h"
+#include "Widgets/Input/SEditableText.h"
+
+#define LOCTEXT_NAMESPACE "STextKeyEditor"
+
+namespace UE::MovieScene
+{
+
+void STextKeyEditor::Construct(const FArguments& InArgs, FTextKeyEditorParams&& InParams)
+{
+	Params = MoveTemp(InParams);
+
+	ChildSlot
+	[
+		SNew(SEditableText)
+		.SelectAllTextWhenFocused(true)
+		.Text(this, &STextKeyEditor::GetText)
+		.OnTextCommitted(this, &STextKeyEditor::OnTextCommitted)
+	];
+
+	InvalidateText();
+}
+
+FText STextKeyEditor::GetText() const
+{
+	return CachedText;
+}
+
+void STextKeyEditor::InvalidateText()
+{
+	ISequencer* Sequencer = Params.WeakSequencer.Pin().Get();
+	UMovieSceneSection* Section = Params.WeakSection.Get();
+	FMovieSceneTextChannel* Channel = Params.ChannelHandle.Get();
+
+	if (!Channel || !Sequencer || !Section)
+	{
+		return;
+	}
+
+	const FFrameTime CurrentTime = UE::MovieScene::ClampToDiscreteRange(Sequencer->GetLocalTime().Time, Section->GetRange());
+
+	// If we have no keys and no default, key with the external value if it exists
+	if (!UE::MovieScene::EvaluateChannel(Section, Channel, CurrentTime, CachedText))
+	{
+		FTrackInstancePropertyBindings* PropertyBindings = Params.WeakPropertyBindings.Pin().Get();
+		for (TWeakObjectPtr<> WeakObject : Sequencer->FindBoundObjects(Params.ObjectBindingID, Sequencer->GetFocusedTemplateID()))
+		{
+			if (UObject* Object = WeakObject.Get())
+			{
+				TOptional<FText> ExternalValue = Params.OnGetExternalValue(*Object, PropertyBindings);
+				if (ExternalValue.IsSet())
+				{
+					CachedText = ExternalValue.GetValue();
+					break;
+				}
+			}
+		}
+	}
+}
+
+void STextKeyEditor::OnTextCommitted(const FText& InText, ETextCommit::Type InCommitType)
+{
+	using namespace UE::MovieScene;
+	using namespace UE::Sequencer;
+
+	UMovieSceneSection* Section = Params.WeakSection.Get();
+	if (!Section)
+	{
+		return;
+	}
+
+	FScopedTransaction Transaction(LOCTEXT("SetTextKey", "Set Text Key Value"));
+
+	CachedText = InText;
+
+	Section->SetFlags(RF_Transactional);
+
+	FMovieSceneTextChannel* Channel  = Params.ChannelHandle.Get();
+	TSharedPtr<ISequencer> Sequencer = Params.WeakSequencer.Pin();
+
+	if (!Section->TryModify() || !Channel || !Sequencer.IsValid())
+	{
+		return;
+	}
+
+	const bool bAutoSetTrackDefaults = Sequencer->GetAutoSetTrackDefaults();
+	const FKeySelection& KeySelection = Sequencer->GetViewModel()->GetSelection()->KeySelection;
+	if (KeySelection.Num())
+	{
+		for (FKeyHandle Key : KeySelection)
+		{
+			// Make sure we only manipulate the values of the channel with the same channel type we're editing
+			TSharedPtr<FChannelModel> ChannelModel = KeySelection.GetModelForKey(Key);
+
+			const bool bMatchingChannelType = ChannelModel.IsValid()
+				&& ChannelModel->GetKeyArea()
+				&& ChannelModel->GetKeyArea()->GetChannel().GetChannelTypeName() == Params.ChannelHandle.GetChannelTypeName();
+
+			if (bMatchingChannelType)
+			{
+				UMovieSceneSection* ModelSection = ChannelModel->GetSection();
+				if (ModelSection && ModelSection->TryModify())
+				{
+					AssignValue(reinterpret_cast<FMovieSceneTextChannel*>(ChannelModel->GetChannel()), Key, InText);
+				}
+			}
+		}
+	}
+	else
+	{
+		const FFrameNumber CurrentTime = Sequencer->GetLocalTime().Time.FloorToFrame();
+
+		EMovieSceneKeyInterpolation Interpolation = GetInterpolationMode(Channel, CurrentTime, Sequencer->GetKeyInterpolation());
+
+		TArray<FKeyHandle> KeysAtCurrentTime;
+		Channel->GetKeys(TRange<FFrameNumber>(CurrentTime), nullptr, &KeysAtCurrentTime);
+
+		if (KeysAtCurrentTime.Num() > 0)
+		{
+			AssignValue(Channel, KeysAtCurrentTime[0], InText);
+		}
+		else
+		{
+			bool bHasAnyKeys = Channel->GetNumKeys() != 0;
+
+			if (bHasAnyKeys || bAutoSetTrackDefaults == false)
+			{
+				// When auto setting track defaults are disabled, add a key even when it's empty so that the changed
+				// value is saved and is propagated to the property.
+				AddKeyToChannel(Channel, CurrentTime, InText, Interpolation);
+				bHasAnyKeys = Channel->GetNumKeys() != 0;
+			}
+
+			if (bHasAnyKeys)
+			{
+				TRange<FFrameNumber> KeyRange     = TRange<FFrameNumber>(CurrentTime);
+				TRange<FFrameNumber> SectionRange = Section->GetRange();
+
+				if (!SectionRange.Contains(KeyRange))
+				{
+					Section->SetRange(TRange<FFrameNumber>::Hull(KeyRange, SectionRange));
+				}
+			}
+		}
+	}
+
+	// Always update the default value when auto-set default values is enabled so that the last changes
+	// are always saved to the track.
+	if (bAutoSetTrackDefaults)
+	{
+		SetChannelDefault(Channel, InText);
+	}
+
+	// need to tell channel change happened (float will call AutoSetTangents())
+	Channel->PostEditChange();
+
+	const FMovieSceneChannelMetaData* ChannelMetaData = Params.ChannelHandle.GetMetaData();
+	Sequencer->OnChannelChanged().Broadcast(ChannelMetaData, Section);
+	Sequencer->NotifyMovieSceneDataChanged(EMovieSceneDataChangeType::TrackValueChangedRefreshImmediately);
+}
+
+}
+#undef LOCTEXT_NAMESPACE
diff --git a/Engine/Plugins/MovieScene/MovieSceneTextTrack/Source/MovieSceneTextTrackEditor/Private/Widgets/STextKeyEditor.h b/Engine/Plugins/MovieScene/MovieSceneTextTrack/Source/MovieSceneTextTrackEditor/Private/Widgets/STextKeyEditor.h
new file mode 100644
--- /dev/null
+++ b/Engine/Plugins/MovieScene/MovieSceneTextTrack/Source/MovieSceneTextTrackEditor/Private/Widgets/STextKeyEditor.h
@@ -0,0 +1,50 @@
+// Copyright Epic Games, Inc. All Rights Reserved.
+
+#pragma once
+
+#include "MovieSceneTextChannel.h"
+#include "Channels/MovieSceneChannelHandle.h"
+#include "Widgets/SCompoundWidget.h"
+
+class ISequencer;
+
+namespace UE::MovieScene
+{
+
+struct FTextKeyEditorParams
+{
+	FGuid ObjectBindingID;
+	TMovieSceneChannelHandle<FMovieSceneTextChannel> ChannelHandle;
+	TWeakObjectPtr<UMovieSceneSection> WeakSection;
+	TWeakPtr<ISequencer> WeakSequencer;
+	TWeakPtr<FTrackInstancePropertyBindings> WeakPropertyBindings;
+	TFunction<TOptional<FText>(UObject&, FTrackInstancePropertyBindings*)> OnGetExternalValue;
+};
+
+class STextKeyEditor : public SCompoundWidget
+{
+public:
+	SLATE_BEGIN_ARGS(STextKeyEditor){}
+	SLATE_END_ARGS()
+
+	void Construct(const FArguments& InArgs, FTextKeyEditorParams&& InParams);
+
+private:
+	/** Gets the CachedText value last set by either OnTextCommitted or InvalidateText*/
+	FText GetText() const;
+
+	/**
+	 * Reevaluates the Text Channel for the given Section at the current sequencer time, and caches the result in CachedText
+	 * Currently only called in Widget Construction
+	 */
+	void InvalidateText();
+
+	/** Updates the CachedText result to the committed text and adds/updates a key to the text channel at the current sequencer time */
+	void OnTextCommitted(const FText& InText, ETextCommit::Type InCommitType);
+
+	FTextKeyEditorParams Params;
+
+	FText CachedText;
+};
+
+}
