Movie Scene Text Track:
- Fix issue with text value widget not evaluating correctly as it was only 'invalidated' at construction
- Fix text widget not being noticeable when value is blank

diff --git a/Engine/Plugins/MovieScene/MovieSceneTextTrack/Source/MovieSceneTextTrackEditor/Private/Widgets/STextKeyEditor.cpp b/Engine/Plugins/MovieScene/MovieSceneTextTrack/Source/MovieSceneTextTrackEditor/Private/Widgets/STextKeyEditor.cpp
--- a/Engine/Plugins/MovieScene/MovieSceneTextTrack/Source/MovieSceneTextTrackEditor/Private/Widgets/STextKeyEditor.cpp
+++ b/Engine/Plugins/MovieScene/MovieSceneTextTrack/Source/MovieSceneTextTrackEditor/Private/Widgets/STextKeyEditor.cpp
@@ -7,7 +7,7 @@
 #include "MVVM/ViewModels/SequencerEditorViewModel.h"
 #include "MovieSceneTimeHelpers.h"
 #include "ScopedTransaction.h"
-#include "Widgets/Input/SEditableText.h"
+#include "Widgets/Input/SEditableTextBox.h"
 
 #define LOCTEXT_NAMESPACE "STextKeyEditor"
 
@@ -20,21 +20,15 @@ void STextKeyEditor::Construct(const FArguments& InArgs, FTextKeyEditorParams&&
 
 	ChildSlot
 	[
-		SNew(SEditableText)
+		SNew(SEditableTextBox)
+		.MinDesiredWidth(10.f)
 		.SelectAllTextWhenFocused(true)
 		.Text(this, &STextKeyEditor::GetText)
 		.OnTextCommitted(this, &STextKeyEditor::OnTextCommitted)
 	];
-
-	InvalidateText();
 }
 
 FText STextKeyEditor::GetText() const
-{
-	return CachedText;
-}
-
-void STextKeyEditor::InvalidateText()
 {
 	ISequencer* Sequencer = Params.WeakSequencer.Pin().Get();
 	UMovieSceneSection* Section = Params.WeakSection.Get();
@@ -42,13 +36,15 @@ void STextKeyEditor::InvalidateText()
 
 	if (!Channel || !Sequencer || !Section)
 	{
-		return;
+		return FText::GetEmpty();
 	}
 
 	const FFrameTime CurrentTime = UE::MovieScene::ClampToDiscreteRange(Sequencer->GetLocalTime().Time, Section->GetRange());
 
+	FText Text;
+
 	// If we have no keys and no default, key with the external value if it exists
-	if (!UE::MovieScene::EvaluateChannel(Section, Channel, CurrentTime, CachedText))
+	if (!UE::MovieScene::EvaluateChannel(Section, Channel, CurrentTime, Text))
 	{
 		FTrackInstancePropertyBindings* PropertyBindings = Params.WeakPropertyBindings.Pin().Get();
 		for (TWeakObjectPtr<> WeakObject : Sequencer->FindBoundObjects(Params.ObjectBindingID, Sequencer->GetFocusedTemplateID()))
@@ -58,12 +54,13 @@ void STextKeyEditor::InvalidateText()
 				TOptional<FText> ExternalValue = Params.OnGetExternalValue(*Object, PropertyBindings);
 				if (ExternalValue.IsSet())
 				{
-					CachedText = ExternalValue.GetValue();
-					break;
+					return ExternalValue.GetValue();
 				}
 			}
 		}
 	}
+
+	return Text;
 }
 
 void STextKeyEditor::OnTextCommitted(const FText& InText, ETextCommit::Type InCommitType)
@@ -79,8 +76,6 @@ void STextKeyEditor::OnTextCommitted(const FText& InText, ETextCommit::Type InCo
 
 	FScopedTransaction Transaction(LOCTEXT("SetTextKey", "Set Text Key Value"));
 
-	CachedText = InText;
-
 	Section->SetFlags(RF_Transactional);
 
 	FMovieSceneTextChannel* Channel  = Params.ChannelHandle.Get();
diff --git a/Engine/Plugins/MovieScene/MovieSceneTextTrack/Source/MovieSceneTextTrackEditor/Private/Widgets/STextKeyEditor.h b/Engine/Plugins/MovieScene/MovieSceneTextTrack/Source/MovieSceneTextTrackEditor/Private/Widgets/STextKeyEditor.h
--- a/Engine/Plugins/MovieScene/MovieSceneTextTrack/Source/MovieSceneTextTrackEditor/Private/Widgets/STextKeyEditor.h
+++ b/Engine/Plugins/MovieScene/MovieSceneTextTrack/Source/MovieSceneTextTrackEditor/Private/Widgets/STextKeyEditor.h
@@ -30,21 +30,13 @@ public:
 	void Construct(const FArguments& InArgs, FTextKeyEditorParams&& InParams);
 
 private:
-	/** Gets the CachedText value last set by either OnTextCommitted or InvalidateText*/
+	/** Gets the evaluated Text value at the current sequencer time */
 	FText GetText() const;
 
-	/**
-	 * Reevaluates the Text Channel for the given Section at the current sequencer time, and caches the result in CachedText
-	 * Currently only called in Widget Construction
-	 */
-	void InvalidateText();
-
 	/** Updates the CachedText result to the committed text and adds/updates a key to the text channel at the current sequencer time */
 	void OnTextCommitted(const FText& InText, ETextCommit::Type InCommitType);
 
 	FTextKeyEditorParams Params;
-
-	FText CachedText;
 };
 
 }
